[UE-216290] Fix crash when changing audio compression formats on SoundWave that's actively playing via looping WavePlayerNode in a metasound patch.

#rnx

diff --git a/Engine/Source/Runtime/VorbisAudioDecoder/Module/Private/VorbisAudioInfo.cpp b/Engine/Source/Runtime/VorbisAudioDecoder/Module/Private/VorbisAudioInfo.cpp
--- a/Engine/Source/Runtime/VorbisAudioDecoder/Module/Private/VorbisAudioInfo.cpp
+++ b/Engine/Source/Runtime/VorbisAudioDecoder/Module/Private/VorbisAudioInfo.cpp
@@ -542,6 +542,12 @@ void FVorbisAudioInfo::SeekToTime( const float SeekTime )
 		UE_LOG(LogAudio, Error, TEXT("FVorbisAudioInfo::SeekToTime failed due to not parsing header first."));
 		return;
 	}
+	
+	if (!VFWrapper->vf.seekable)
+	{
+		UE_LOG(LogAudio, Error, TEXT("FVorbisAudioInfo::SeekToTime failed due to vorbis file not being seekable"));
+		return;
+	}
 
 	const float TargetTime = FMath::Min(SeekTime, (float)ov_time_total(&VFWrapper->vf, -1));
 	ov_time_seek( &VFWrapper->vf, TargetTime );
@@ -563,6 +569,12 @@ void FVorbisAudioInfo::SeekToFrame(const uint32 SeekFrames)
 		return;
 	}
 
+	if (!VFWrapper->vf.seekable)
+	{
+		UE_LOG(LogAudio, Error, TEXT("FVorbisAudioInfo::SeekToTime failed due to vorbis file not being seekable"));
+		return;
+	}
+
 	vorbis_info* vi = ov_info(&VFWrapper->vf, -1);
 	int64 SeekSamples = SeekFrames * vi->channels;
 	const int64 TargetSample = FMath::Min(SeekSamples, (int64)ov_pcm_total(&VFWrapper->vf, -1));
