[Audio] Fix thread-safety races in vorbis DLL init and AudioFormatOgg singleton

#Synced-CL 51959729

diff --git a/Engine/Source/Developer/AudioFormatOgg/Private/AudioFormatOgg.cpp b/Engine/Source/Developer/AudioFormatOgg/Private/AudioFormatOgg.cpp
--- a/Engine/Source/Developer/AudioFormatOgg/Private/AudioFormatOgg.cpp
+++ b/Engine/Source/Developer/AudioFormatOgg/Private/AudioFormatOgg.cpp
@@ -528,24 +528,14 @@ public:
  * Module for ogg audio compression
  */
 
-static IAudioFormat* Singleton = NULL;
-
 class FAudioPlatformOggModule : public IAudioFormatModule
 {
 public:
-	virtual ~FAudioPlatformOggModule()
-	{
-		delete Singleton;
-		Singleton = NULL;
-	}
 	virtual IAudioFormat* GetAudioFormat()
 	{
-		if (!Singleton)
-		{
-			LoadVorbisLibraries();
-			Singleton = new FAudioFormatOgg();
-		}
-		return Singleton;
+		LoadVorbisLibraries();
+		static FAudioFormatOgg Instance;
+		return &Instance;
 	}
 };
 
diff --git a/Engine/Source/Runtime/VorbisAudioDecoder/Module/Private/VorbisAudioInfo.cpp b/Engine/Source/Runtime/VorbisAudioDecoder/Module/Private/VorbisAudioInfo.cpp
--- a/Engine/Source/Runtime/VorbisAudioDecoder/Module/Private/VorbisAudioInfo.cpp
+++ b/Engine/Source/Runtime/VorbisAudioDecoder/Module/Private/VorbisAudioInfo.cpp
@@ -789,10 +789,8 @@ bool FVorbisAudioInfo::HasError() const
 
 void LoadVorbisLibraries()
 {
-	static bool bIsInitialized = false;
-	if (!bIsInitialized)
+	static const bool bDone = []() -> bool
 	{
-		bIsInitialized = true;
 #if WITH_OGGVORBIS_DLL && WITH_OGGVORBIS
 		//@todo if ogg is every ported to another platform, then use the platform abstraction to load these DLLs
 		// Load the Ogg dlls
@@ -846,7 +844,9 @@ void LoadVorbisLibraries()
 #elif WITH_OGGVORBIS
 		bDllLoaded = true;
 #endif	//WITH_OGGVORBIS_DLL
-	}
+		return true;
+	}();
+	(void)bDone;
 }
 
 #endif		// WITH_OGGVORBIS
